UE4 - Remove some redundant FName calls

diff --git a/Engine/Source/Developer/PackageDependencyInfo/Private/PackageDependencyInfo.cpp b/Engine/Source/Developer/PackageDependencyInfo/Private/PackageDependencyInfo.cpp
--- a/Engine/Source/Developer/PackageDependencyInfo/Private/PackageDependencyInfo.cpp
+++ b/Engine/Source/Developer/PackageDependencyInfo/Private/PackageDependencyInfo.cpp
@@ -383,14 +383,13 @@ void FPackageDependencyInfo::RecursiveDeterminePackageDependentTimeStamp(const T
 			PkgInfo->bContainsMap = Linker->ContainsMap();
 			PkgInfo->bContainsBlueprints = Linker->ContainsCode();
 
-			FName CheckMaterial = FName(TEXT("Material"));
-			FName CheckMIC = FName(TEXT("MaterialInstanceConstant"));
-			FName CheckMID = FName(TEXT("MaterialInstanceDynamic"));
-			FName CheckLMIC = FName(TEXT("LandscapeMaterialInstanceConstant"));
-			FName CheckWorld = FName(TEXT("World"));
-			FName CheckBlueprint = FName(TEXT("Blueprint"));
-			FName CheckAnimBlueprint = FName(TEXT("AnimBlueprint"));
-
+			static const FName CheckMaterial(TEXT("Material"));
+			static const FName CheckMIC(TEXT("MaterialInstanceConstant"));
+			static const FName CheckMID(TEXT("MaterialInstanceDynamic"));
+			static const FName CheckLMIC(TEXT("LandscapeMaterialInstanceConstant"));
+			static const FName CheckWorld(TEXT("World"));
+			static const FName CheckBlueprint(TEXT("Blueprint"));
+			static const FName CheckAnimBlueprint(TEXT("AnimBlueprint"));
 
 			// Check the export map for material interfaces
 			for (int32 ExpIdx = 0; ExpIdx < Linker->ExportMap.Num(); ExpIdx++)
